Fixed a shadow variable warning in SPluginCreatorTabContent

diff --git a/Engine/Plugins/Editor/PluginCreator/Source/PluginCreator/Private/UI/SPluginCreatorTabContent.h b/Engine/Plugins/Editor/PluginCreator/Source/PluginCreator/Private/UI/SPluginCreatorTabContent.h
--- a/Engine/Plugins/Editor/PluginCreator/Source/PluginCreator/Private/UI/SPluginCreatorTabContent.h
+++ b/Engine/Plugins/Editor/PluginCreator/Source/PluginCreator/Private/UI/SPluginCreatorTabContent.h
@@ -123,9 +123,6 @@ private:
 	/** Name of the plugin you want to create*/
 	FText PluginNameText;
 
-	/** Plugin descriptor that will be used to create .uplugin file*/
-	FPluginDescriptor Descriptor;
-
 	/** Property panel used to display plugin descriptor parameters*/
 	TSharedPtr<class IDetailsView> PropertyView;
 
