Remove the bEnabledByDefault option for new plugins. It's designed for built-in plugins only; all project plugins are enabled by default anyway.

diff --git a/Engine/Plugins/Editor/PluginCreator/Source/PluginCreator/Private/Helpers/PluginDescriptorObject.cpp b/Engine/Plugins/Editor/PluginCreator/Source/PluginCreator/Private/Helpers/PluginDescriptorObject.cpp
--- a/Engine/Plugins/Editor/PluginCreator/Source/PluginCreator/Private/Helpers/PluginDescriptorObject.cpp
+++ b/Engine/Plugins/Editor/PluginCreator/Source/PluginCreator/Private/Helpers/PluginDescriptorObject.cpp
@@ -12,13 +12,11 @@ UPluginDescriptorObject::UPluginDescriptorObject(const FObjectInitializer& Objec
 	Version = 1;
 	VersionName = TEXT("1.0");
 	Category = TEXT("Other");
-	bEnabledByDefault = true;
 }
 
 void UPluginDescriptorObject::FillDescriptor(FPluginDescriptor& OutDescriptor)
 {
 	OutDescriptor.bCanContainContent = bCanContainContent;
-	OutDescriptor.bEnabledByDefault = bEnabledByDefault;
 	OutDescriptor.bIsBetaVersion = bIsBetaVersion;
 	OutDescriptor.Category = Category;
 	OutDescriptor.CreatedBy = CreatedBy;
diff --git a/Engine/Plugins/Editor/PluginCreator/Source/PluginCreator/Private/Helpers/PluginDescriptorObject.h b/Engine/Plugins/Editor/PluginCreator/Source/PluginCreator/Private/Helpers/PluginDescriptorObject.h
--- a/Engine/Plugins/Editor/PluginCreator/Source/PluginCreator/Private/Helpers/PluginDescriptorObject.h
+++ b/Engine/Plugins/Editor/PluginCreator/Source/PluginCreator/Private/Helpers/PluginDescriptorObject.h
@@ -67,10 +67,6 @@ public:
 	/** List of all modules associated with this plugin */
 	TArray<FModuleDescriptor> Modules;
 
-	/** Whether this plugin should be enabled by default for all projects */
-	UPROPERTY(EditAnywhere, Category = "PluginDescription")
-	bool bEnabledByDefault;
-
 	/** Can this plugin contain content? */
 	UPROPERTY(EditAnywhere, Category = "PluginDescription")
 	bool bCanContainContent;
