Fixes for Plugin creator plugin in non-unity mode

diff --git a/Engine/Plugins/Editor/PluginCreator/Source/PluginCreator/Private/PluginCreatorPlugin.cpp b/Engine/Plugins/Editor/PluginCreator/Source/PluginCreator/Private/PluginCreatorPlugin.cpp
--- a/Engine/Plugins/Editor/PluginCreator/Source/PluginCreator/Private/PluginCreatorPlugin.cpp
+++ b/Engine/Plugins/Editor/PluginCreator/Source/PluginCreator/Private/PluginCreatorPlugin.cpp
@@ -17,6 +17,8 @@ static const FName PCPluginTabName("PluginCreatorPlugin");
 
 #include "PluginCreatorCommands.h"
 
+DEFINE_LOG_CATEGORY(PluginCreatorPluginLog);
+
 void FPluginCreatorModule::StartupModule()
 {
 	FPluginCreatorStyle::Initialize();
diff --git a/Engine/Plugins/Editor/PluginCreator/Source/PluginCreator/Private/PluginCreatorPlugin.h b/Engine/Plugins/Editor/PluginCreator/Source/PluginCreator/Private/PluginCreatorPlugin.h
--- a/Engine/Plugins/Editor/PluginCreator/Source/PluginCreator/Private/PluginCreatorPlugin.h
+++ b/Engine/Plugins/Editor/PluginCreator/Source/PluginCreator/Private/PluginCreatorPlugin.h
@@ -8,7 +8,6 @@
 
 
 DECLARE_LOG_CATEGORY_EXTERN(PluginCreatorPluginLog, Log, All);
-DEFINE_LOG_CATEGORY(PluginCreatorPluginLog);
 
 class FPluginCreatorModule : public IModuleInterface
 {
diff --git a/Engine/Plugins/Editor/PluginCreator/Source/PluginCreator/Private/UI/SPluginCreatorTabContent.cpp b/Engine/Plugins/Editor/PluginCreator/Source/PluginCreator/Private/UI/SPluginCreatorTabContent.cpp
--- a/Engine/Plugins/Editor/PluginCreator/Source/PluginCreator/Private/UI/SPluginCreatorTabContent.cpp
+++ b/Engine/Plugins/Editor/PluginCreator/Source/PluginCreator/Private/UI/SPluginCreatorTabContent.cpp
@@ -18,6 +18,8 @@
 
 #include "SNotificationList.h"
 #include "NotificationManager.h"
+#include "PluginCreatorStyle.h"
+#include "PluginCreatorPlugin.h"
 
 #define LOCTEXT_NAMESPACE "PluginCreatorPluginModule"
 
