PR #1145: Enabled engine/game plugin switch in plugin wizard (Contributed by [REDACTED_PERSON])

Made slight change to make sure that Ed Mode information is not always written into the module .cpp file.

diff --git a/Engine/Plugins/Editor/PluginCreator/Source/PluginCreator/PluginCreator.Build.cs b/Engine/Plugins/Editor/PluginCreator/Source/PluginCreator/PluginCreator.Build.cs
--- a/Engine/Plugins/Editor/PluginCreator/Source/PluginCreator/PluginCreator.Build.cs
+++ b/Engine/Plugins/Editor/PluginCreator/Source/PluginCreator/PluginCreator.Build.cs
@@ -20,6 +20,7 @@ public class PluginCreator : ModuleRules
 				"Core",
 				"CoreUObject",
 				"Slate",
+                "Engine",
 				"SlateCore",
  				"InputCore",
 				"UnrealEd",
diff --git a/Engine/Plugins/Editor/PluginCreator/Source/PluginCreator/Private/Helpers/PluginDescriptorDetails.cpp b/Engine/Plugins/Editor/PluginCreator/Source/PluginCreator/Private/Helpers/PluginDescriptorDetails.cpp
new file mode 100644
--- /dev/null
+++ b/Engine/Plugins/Editor/PluginCreator/Source/PluginCreator/Private/Helpers/PluginDescriptorDetails.cpp
@@ -0,0 +1,24 @@
+#include "PluginCreatorPrivatePCH.h"
+#include "PluginDescriptorDetails.h"
+#include "PluginDescriptorObject.h"
+
+#define LOCTEXT_NAMESPACE "PluginDescriptorDetails"
+
+TSharedRef<IDetailCustomization> FPluginDescriptorDetails::MakeInstance()
+{
+	return MakeShareable(new FPluginDescriptorDetails);
+}
+
+void FPluginDescriptorDetails::CustomizeDetails(IDetailLayoutBuilder& DetailLayout)
+{
+	TSharedPtr<IPropertyHandle> IsEnginePluginProperty = DetailLayout.GetProperty(GET_MEMBER_NAME_CHECKED(UPluginDescriptorObject, bIsEnginePlugin));
+
+	// Hide bIsEnginePlugin flag if this is a Launcher build
+	// User should be able to create Game plugins only if thats the case
+	if (FApp::IsEngineInstalled())
+	{
+		IsEnginePluginProperty->MarkHiddenByCustomization();
+	}
+}
+
+#undef LOCTEXT_NAMESPACE
\ No newline at end of file
diff --git a/Engine/Plugins/Editor/PluginCreator/Source/PluginCreator/Private/Helpers/PluginDescriptorDetails.h b/Engine/Plugins/Editor/PluginCreator/Source/PluginCreator/Private/Helpers/PluginDescriptorDetails.h
new file mode 100644
--- /dev/null
+++ b/Engine/Plugins/Editor/PluginCreator/Source/PluginCreator/Private/Helpers/PluginDescriptorDetails.h
@@ -0,0 +1,11 @@
+#pragma once
+
+class FPluginDescriptorDetails : public IDetailCustomization
+{
+public:
+	/** Makes a new instance of this detail layout class for a specific detail view requesting it */
+	static TSharedRef<IDetailCustomization> MakeInstance();
+
+	/** IDetailCustomization interface */
+	virtual void CustomizeDetails(IDetailLayoutBuilder& DetailLayout) override;
+};
diff --git a/Engine/Plugins/Editor/PluginCreator/Source/PluginCreator/Private/Helpers/PluginDescriptorObject.cpp b/Engine/Plugins/Editor/PluginCreator/Source/PluginCreator/Private/Helpers/PluginDescriptorObject.cpp
--- a/Engine/Plugins/Editor/PluginCreator/Source/PluginCreator/Private/Helpers/PluginDescriptorObject.cpp
+++ b/Engine/Plugins/Editor/PluginCreator/Source/PluginCreator/Private/Helpers/PluginDescriptorObject.cpp
@@ -12,6 +12,12 @@ UPluginDescriptorObject::UPluginDescriptorObject(const FObjectInitializer& Objec
 	Version = 1;
 	VersionName = TEXT("1.0");
 	Category = TEXT("Other");
+
+	
+	PrivateDependencyModuleNames.Add("CoreUObject");
+	PrivateDependencyModuleNames.Add("Engine");
+	PrivateDependencyModuleNames.Add("Slate");
+	PrivateDependencyModuleNames.Add("SlateCore");
 }
 
 void UPluginDescriptorObject::FillDescriptor(FPluginDescriptor& OutDescriptor)
diff --git a/Engine/Plugins/Editor/PluginCreator/Source/PluginCreator/Private/Helpers/PluginDescriptorObject.h b/Engine/Plugins/Editor/PluginCreator/Source/PluginCreator/Private/Helpers/PluginDescriptorObject.h
--- a/Engine/Plugins/Editor/PluginCreator/Source/PluginCreator/Private/Helpers/PluginDescriptorObject.h
+++ b/Engine/Plugins/Editor/PluginCreator/Source/PluginCreator/Private/Helpers/PluginDescriptorObject.h
@@ -6,6 +6,58 @@
 #include "PluginDescriptor.h"
 #include "PluginDescriptorObject.generated.h"
 
+/* Struct that holds basic Editor mode parameters*/
+USTRUCT()
+struct FEdModeSettings
+{
+	GENERATED_USTRUCT_BODY()
+
+	FEdModeSettings(){}
+	~FEdModeSettings(){}
+
+	// Name of your editor mode
+	// It will be filled automatically if you specify plugin name
+	UPROPERTY(EditAnywhere, Category = "EdModeSettings")
+	FString Name;
+
+	// Should this mode use toolkit?
+	// Toolkits can specify UI that will appear in "Modes" tab (next to Foliage, Landscape etc)
+	UPROPERTY(EditAnywhere, Category = "EdModeSettings")
+	bool bUsesToolkits;
+
+	// Do you want this mode to include very basic UI that demonstrates editor interaction and undo/redo functions usage?
+	UPROPERTY(EditAnywhere, Category = "EdModeSettings", meta = (editcondition = "bUsesToolkits"))
+	bool bIncludeSampleUI;
+
+	/*
+	*	Make string that will register mode.
+	*	@return - string which will be used to register mode in PluginModule.cpp (template) or empty string if there is no mode
+	*/
+	FString MakeRegisterModeString()
+	{
+		if (Name.IsEmpty())
+		{
+			return FString("");
+		}
+
+		return FString::Printf(TEXT("FEditorModeRegistry::Get().RegisterMode<F%s>(F%s::EM_%sId, LOCTEXT(\"%sName\", \"%s\"), FSlateIcon(), %s);"), *Name, *Name, *Name, *Name, *Name, bUsesToolkits ? TEXT("true") : TEXT("false"));
+	}
+
+	/*
+	*	Make string that will register mode.
+	*	@return - string which will be used to unregister mode in PluginModule.cpp (template) or empty string if there is no mode
+	*/
+	FString MakeUnRegisterModeString()
+	{
+		if (Name.IsEmpty())
+		{
+			return FString("");
+		}
+
+		return FString::Printf(TEXT("FEditorModeRegistry::Get().UnregisterMode(F%s::EM_%sId);"), *Name, *Name);
+	}
+};
+
 /**
 *	We use this object to display plugin properties using details view.
 */
@@ -16,7 +68,13 @@ class UPluginDescriptorObject : public UObject
 
 	UPluginDescriptorObject(const FObjectInitializer& ObjectInitializer);
 
-public:
+public:	
+
+	/** Should new plugin be considered as 'Engine' plugin?
+	If so, it will be put under Engine/Plugins directory. */
+	UPROPERTY(EditAnywhere, Category = "CreationOptions")
+	bool bIsEnginePlugin;
+
 	// Should the project files be regenerated when the plugin is created?
 	UPROPERTY(EditAnywhere, Category = "CreationOptions")
 	bool bRegenerateProjectFiles;
@@ -25,6 +83,23 @@ public:
 	UPROPERTY(EditAnywhere, Category = "CreationOptions")
 	bool bShowPluginFolder;
 
+	// Do you want to create Editor mode within this plugin?
+	UPROPERTY(EditAnywhere, Category = "CreationOptions", AdvancedDisplay)
+	bool bMakeEditorMode;
+
+	UPROPERTY(EditAnywhere, Category = "CreationOptions", AdvancedDisplay, meta = (editcondition = "bMakeEditorMode"))
+	FEdModeSettings EditorModeSettings;
+
+	// Do you want to create blueprint function library within this plugin?
+	// Blueprint function library lets you create your own, static blueprint nodes
+	UPROPERTY(EditAnywhere, Category = "CreationOptions", AdvancedDisplay)
+	bool bAddBPLibrary;
+
+	// Modify this array to set Private Dependencies for this plugin
+	// Do it only if you know what you're doing! Missing dependencies can cause build failure!
+	UPROPERTY(EditAnywhere, Category = "CreationOptions", AdvancedDisplay)
+	TArray<FString> PrivateDependencyModuleNames;
+
 	/** Descriptor version number */
 	UPROPERTY(VisibleAnywhere, Category = "PluginDescription")
 	int32 FileVersion;
diff --git a/Engine/Plugins/Editor/PluginCreator/Source/PluginCreator/Private/Helpers/PluginHelpers.cpp b/Engine/Plugins/Editor/PluginCreator/Source/PluginCreator/Private/Helpers/PluginHelpers.cpp
--- a/Engine/Plugins/Editor/PluginCreator/Source/PluginCreator/Private/Helpers/PluginHelpers.cpp
+++ b/Engine/Plugins/Editor/PluginCreator/Source/PluginCreator/Private/Helpers/PluginHelpers.cpp
@@ -20,7 +20,7 @@ bool FPluginHelpers::ReadTemplateFile(const FString& TemplateFileName, FString&
 	return false;
 }
 
-bool FPluginHelpers::CreatePluginBuildFile(const FString& NewBuildFileName, const FString& PluginName, FText& OutFailReason, FString TemplateType)
+bool FPluginHelpers::CreatePluginBuildFile(const FString& NewBuildFileName, const FString& PluginName, FText& OutFailReason, FString TemplateType, TArray<FString> InPrivateDependencyModuleNames)
 {
 	FString Template;
 
@@ -37,9 +37,18 @@ bool FPluginHelpers::CreatePluginBuildFile(const FString& NewBuildFileName, cons
 	TArray<FString> PrivateDependencyModuleNames;
 	TArray<FString> DynamicallyLoadedModuleNames;
 
-	PrivateDependencyModuleNames.Add("Slate");
-	PrivateDependencyModuleNames.Add("SlateCore");
-
+	if (InPrivateDependencyModuleNames.Num() == 0)
+	{
+		PrivateDependencyModuleNames.Add("CoreUObject");
+		PrivateDependencyModuleNames.Add("Engine");
+		PrivateDependencyModuleNames.Add("Slate");
+		PrivateDependencyModuleNames.Add("SlateCore");
+	}
+	else
+	{
+		PrivateDependencyModuleNames = InPrivateDependencyModuleNames;
+	}
+	
 	FString FinalOutput = Template.Replace(TEXT("%PUBLIC_INCLUDE_PATHS_NAMES%"), *GameProjectUtils::MakeCommaDelimitedList(PublicIncludePathsNames), ESearchCase::CaseSensitive);
 	FinalOutput = FinalOutput.Replace(TEXT("%PRIVATE_INCLUDE_PATHS_NAMES%"), *GameProjectUtils::MakeCommaDelimitedList(PrivateIncludePathsNames), ESearchCase::CaseSensitive);
 	FinalOutput = FinalOutput.Replace(TEXT("%PUBLIC_DEPENDENCY_MODULE_NAMES%"), *GameProjectUtils::MakeCommaDelimitedList(PublicDependencyModuleNames), ESearchCase::CaseSensitive);
@@ -69,7 +78,7 @@ bool FPluginHelpers::CreatePluginHeaderFile(const FString& FolderPath, const FSt
 	return GameProjectUtils::WriteOutputFile(FolderPath / PluginName + TEXT(".h"), FinalOutput, OutFailReason);
 }
 
-bool FPluginHelpers::CreatePluginCPPFile(const FString& FolderPath, const FString& PluginName, FText& OutFailReason, FString TemplateType)
+bool FPluginHelpers::CreatePluginCPPFile(const FString& FolderPath, const FString& PluginName, FText& OutFailReason, FString TemplateType, bool bMakeEditorMode, FEdModeSettings EdModeSettings)
 {
 	FString Template;
 
@@ -81,10 +90,22 @@ bool FPluginHelpers::CreatePluginCPPFile(const FString& FolderPath, const FStrin
 	}
 
 	TArray<FString> PublicHeaderIncludes;
+	FString RegisterModeString;
+	FString UnRegisterModeString;
+
+	if (bMakeEditorMode && !EdModeSettings.Name.IsEmpty())
+	{
+		PublicHeaderIncludes.Add(FString::Printf(TEXT("%s.h"), *(EdModeSettings.Name)));
+		RegisterModeString = EdModeSettings.MakeRegisterModeString();
+		UnRegisterModeString = EdModeSettings.MakeUnRegisterModeString();
+	}
 
 	FString	FinalOutput = Template.Replace(TEXT("%PUBLIC_HEADER_INCLUDES%"), *GameProjectUtils::MakeIncludeList(PublicHeaderIncludes), ESearchCase::CaseSensitive);
 	FinalOutput = FinalOutput.Replace(TEXT("%PLUGIN_NAME%"), *PluginName, ESearchCase::CaseSensitive);
 
+	FinalOutput = FinalOutput.Replace(TEXT("%REGISTER_EDMODE%"), *RegisterModeString, ESearchCase::CaseSensitive);
+	FinalOutput = FinalOutput.Replace(TEXT("%UNREGISTER_EDMODE%"), *UnRegisterModeString, ESearchCase::CaseSensitive);
+
 	return GameProjectUtils::WriteOutputFile(FolderPath / PluginName + TEXT(".cpp"), FinalOutput, OutFailReason);
 }
 
@@ -105,6 +126,46 @@ bool FPluginHelpers::CreatePrivatePCHFile(const FString& FolderPath, const FStri
 	return GameProjectUtils::WriteOutputFile(FolderPath / PluginName + TEXT("PrivatePCH.h"), FinalOutput, OutFailReason);
 }
 
+bool FPluginHelpers::CreateBlueprintFunctionLibraryFiles(const FString& PrivateFolderPath, const FString& PublicFolderPath, const FString& PluginName, FText& OutFailReason)
+{
+	{
+		FString Template;
+		FString TemplateFileName = TEXT("BlueprintFunctionLibrary.h.template");
+
+		if (!ReadTemplateFile(TemplateFileName, Template, OutFailReason))
+		{
+			return false;
+		}
+
+		FString	FinalOutput = Template.Replace(TEXT("%PLUGIN_NAME%"), *PluginName, ESearchCase::CaseSensitive);
+
+		if (!GameProjectUtils::WriteOutputFile(PublicFolderPath / PluginName + TEXT("BPLibrary.h"), FinalOutput, OutFailReason))
+		{
+			return false;
+		}
+	}
+
+	{
+		FString Template;
+		FString TemplateFileName = TEXT("BlueprintFunctionLibrary.cpp.template");
+
+		if (!ReadTemplateFile(TemplateFileName, Template, OutFailReason))
+		{
+			return false;
+		}
+
+		FString	FinalOutput = Template.Replace(TEXT("%PLUGIN_NAME%"), *PluginName, ESearchCase::CaseSensitive);
+
+		if (!GameProjectUtils::WriteOutputFile(PrivateFolderPath / PluginName + TEXT("BPLibrary.cpp"), FinalOutput, OutFailReason))
+		{
+			return false;
+		}
+	}
+	
+
+	return true;
+}
+
 bool FPluginHelpers::CreatePluginStyleFiles(const FString& PrivateFolderPath, const FString& PublicFolderPath, const FString& PluginName, FText& OutFailReason, FString TemplateType)
 {
 	{
@@ -154,6 +215,115 @@ bool FPluginHelpers::CreatePluginStyleFiles(const FString& PrivateFolderPath, co
 	
 }
 
+bool FPluginHelpers::CreatePluginEdModeFiles(const FString& PrivateFolderPath, const FString& PublicFolderPath, const FString& PluginName, FEdModeSettings EdModeSettings, FText& OutFailReason, FString TemplateType /*= FString("")*/)
+{
+	FString EdModeName = EdModeSettings.Name;
+
+	{
+		// Make EdMode.h file first
+
+		FString Template;
+		FString TemplateFileName = TEXT("EdMode.h.template");
+
+		if (!ReadTemplateFile(TemplateType.IsEmpty() ? TemplateFileName : TemplateType / TemplateFileName, Template, OutFailReason))
+		{
+			return false;
+		}
+
+		FString	FinalOutput = Template.Replace(TEXT("%PLUGIN_NAME%"), *PluginName, ESearchCase::CaseSensitive);
+		FinalOutput = FinalOutput.Replace(TEXT("%MODE_NAME%"), *EdModeName, ESearchCase::CaseSensitive);
+
+		if (!GameProjectUtils::WriteOutputFile(PublicFolderPath / EdModeName + TEXT(".h"), FinalOutput, OutFailReason))
+		{
+			return false;
+		}
+	}
+
+	{
+		// Now .cpp file
+
+		FString Template;
+
+		FString TemplateFileName = TEXT("EdMode.cpp.template");
+
+		if (!ReadTemplateFile(TemplateType.IsEmpty() ? TemplateFileName : TemplateType / TemplateFileName, Template, OutFailReason))
+		{
+			return false;
+		}
+
+		FString UsesToolkit = EdModeSettings.bUsesToolkits ? TEXT("true") : TEXT("false");
+		FString BeginComment = EdModeSettings.bUsesToolkits ? TEXT("") : TEXT("/*");
+		FString EndComment = EdModeSettings.bUsesToolkits ? TEXT("") : TEXT("*/");
+
+		FString	FinalOutput = Template.Replace(TEXT("%PLUGIN_NAME%"), *PluginName, ESearchCase::CaseSensitive);
+		FinalOutput = FinalOutput.Replace(TEXT("%MODE_NAME%"), *EdModeName, ESearchCase::CaseSensitive);
+		FinalOutput = FinalOutput.Replace(TEXT("%BEGIN_BLOCK_COMMENT%"), *BeginComment, ESearchCase::CaseSensitive);
+		FinalOutput = FinalOutput.Replace(TEXT("%END_BLOCK_COMMENT%"), *EndComment, ESearchCase::CaseSensitive);
+		FinalOutput = FinalOutput.Replace(TEXT("%USES_TOOLKIT%"), *UsesToolkit, ESearchCase::CaseSensitive);
+
+
+		if (!GameProjectUtils::WriteOutputFile(PrivateFolderPath / EdModeName + TEXT(".cpp"), FinalOutput, OutFailReason))
+		{
+			return false;
+		}
+	}
+
+	// Make toolkit files if necessary
+	if (EdModeSettings.bUsesToolkits)
+	{
+		{
+			FString Template;
+			FString TemplateFileName = TEXT("Toolkit.h.template");
+
+			if (!ReadTemplateFile(TemplateType.IsEmpty() ? TemplateFileName : TemplateType / TemplateFileName, Template, OutFailReason))
+			{
+				return false;
+			}
+
+			FString	FinalOutput = Template.Replace(TEXT("%PLUGIN_NAME%"), *PluginName, ESearchCase::CaseSensitive);
+			FinalOutput = FinalOutput.Replace(TEXT("%MODE_NAME%"), *EdModeName, ESearchCase::CaseSensitive);
+
+			if (!GameProjectUtils::WriteOutputFile(PublicFolderPath / EdModeName + TEXT("Toolkit.h"), FinalOutput, OutFailReason))
+			{
+				return false;
+			}
+		}
+
+		{
+			FString Template;
+			FString TemplateFileName = TEXT("Toolkit.cpp.template");
+
+			if (!ReadTemplateFile(TemplateType.IsEmpty() ? TemplateFileName : TemplateType / TemplateFileName, Template, OutFailReason))
+			{
+				return false;
+			}
+
+			FString SampleUITemplate = TEXT("");
+
+			if (EdModeSettings.bIncludeSampleUI)
+			{
+				TemplateFileName = TEXT("ToolkitUI.template");
+
+				if (!ReadTemplateFile(TemplateFileName, SampleUITemplate, OutFailReason))
+				{
+					return false;
+				}
+			}
+
+			FString	FinalOutput = Template.Replace(TEXT("%PLUGIN_NAME%"), *PluginName, ESearchCase::CaseSensitive);
+			FinalOutput = FinalOutput.Replace(TEXT("%MODE_NAME%"), *EdModeName, ESearchCase::CaseSensitive);
+			FinalOutput = FinalOutput.Replace(TEXT("%TOOLKIT_SAMPLE_UI%"), *SampleUITemplate, ESearchCase::CaseSensitive);
+
+			if (!GameProjectUtils::WriteOutputFile(PrivateFolderPath / EdModeName + TEXT("Toolkit.cpp"), FinalOutput, OutFailReason))
+			{
+				return false;
+			}
+		}
+	}
+
+	return true;
+}
+
 bool FPluginHelpers::CreateCommandsFiles(const FString& FolderPath, const FString& PluginName, FText& OutFailReason, FString TemplateType)
 {
 	{
diff --git a/Engine/Plugins/Editor/PluginCreator/Source/PluginCreator/Private/Helpers/PluginHelpers.h b/Engine/Plugins/Editor/PluginCreator/Source/PluginCreator/Private/Helpers/PluginHelpers.h
--- a/Engine/Plugins/Editor/PluginCreator/Source/PluginCreator/Private/Helpers/PluginHelpers.h
+++ b/Engine/Plugins/Editor/PluginCreator/Source/PluginCreator/Private/Helpers/PluginHelpers.h
@@ -2,6 +2,7 @@
 
 #pragma once
 
+#include "PluginDescriptorObject.h"
 #include "PluginDescriptor.h"
 
 class FPluginHelpers
@@ -19,8 +20,9 @@ public:
 	*	@param PluginName - name of the plugin this file belongs to
 	*	@param OutFailReason - fail reason text if something goes wrong
 	*	@param TemplateType - empty by default; can be Blank, Basic or Advanced
+	*	@param InPrivateDependencyModuleNames - Array of module names that should be included in this plugin. If it's empty, basic modules will be included (see definition of this function).
 	*	@return true if file was created, false otherwise*/
-	static bool CreatePluginBuildFile(const FString& NewBuildFileName, const FString& PluginName, FText& OutFailReason, FString TemplateType = FString(""));
+	static bool CreatePluginBuildFile(const FString& NewBuildFileName, const FString& PluginName, FText& OutFailReason, FString TemplateType = FString(""), TArray<FString> InPrivateDependencyModuleNames = TArray<FString>());
 
 	/** Create plugin header file.
 	*	@param FolderPath - filename with full path
@@ -39,13 +41,17 @@ public:
 	*	@return true if file was created, false otherwise*/
 	static bool CreatePluginStyleFiles(const FString& PrivateFolderPath, const FString& PublicFolderPath, const FString& PluginName, FText& OutFailReason, FString TemplateType = FString(""));
 
+	static bool CreatePluginEdModeFiles(const FString& PrivateFolderPath, const FString& PublicFolderPath, const FString& PluginName,FEdModeSettings EdModeSettings, FText& OutFailReason, FString TemplateType = FString(""));
+
 	/** Create plugin source file.
 	*	@param FolderPath - filename with full path
 	*	@param PluginName - name of the plugin this file belongs to
 	*	@param OutFailReason - fail reason text if something goes wrong
 	*	@param TemplateType - empty by default; can be Blank, Basic or Advanced
+	*	@param bMakeEditorMode - Whether we are making cpp file for plugin with Editor Mode
+	*	@param EdModeSettings - The Editor Mode settings
 	*	@return true if file was created, false otherwise*/
-	static bool CreatePluginCPPFile(const FString& FolderPath, const FString& PluginName, FText& OutFailReason, FString TemplateType = FString(""));
+	static bool CreatePluginCPPFile(const FString& FolderPath, const FString& PluginName, FText& OutFailReason, FString TemplateType = FString(""), bool bMakeEditorMode = false, FEdModeSettings EdModeSettings = FEdModeSettings());
 
 	/** Create plugin UI commands file.
 	*	@param FolderPath - filename with full path
@@ -62,4 +68,12 @@ public:
 	*	@param TemplateType - empty by default; can be Blank, Basic or Advanced
 	*	@return true if file was created, false otherwise*/
 	static bool CreatePrivatePCHFile(const FString& FolderPath, const FString& PluginName, FText& OutFailReason, FString TemplateType = FString(""));
+	/** Create Blueprint Function library class file.
+	*	@param PrivateFolderPath - Private folder path for source file
+	*	@param PublicFolderPath - Public folder path for header file
+	*	@param PluginName - name of the plugin this file belongs to
+	*	@param OutFailReason - fail reason text if something goes wrong
+	*	@return true if file was created, false otherwise*/
+	static bool CreateBlueprintFunctionLibraryFiles(const FString& PrivateFolderPath, const FString& PublicFolderPath, const FString& PluginName, FText& OutFailReason);
+
 };
\ No newline at end of file
diff --git a/Engine/Plugins/Editor/PluginCreator/Source/PluginCreator/Private/PluginCreatorPlugin.cpp b/Engine/Plugins/Editor/PluginCreator/Source/PluginCreator/Private/PluginCreatorPlugin.cpp
--- a/Engine/Plugins/Editor/PluginCreator/Source/PluginCreator/Private/PluginCreatorPlugin.cpp
+++ b/Engine/Plugins/Editor/PluginCreator/Source/PluginCreator/Private/PluginCreatorPlugin.cpp
@@ -6,9 +6,7 @@
 #include "SPluginCreatorTabContent.h"
 #include "SlateExtras.h"
 #include "PluginCreatorStyle.h"
-
-#include "SLayoutCreator.h"
-
+#include "PluginDescriptorDetails.h"
 #include "MultiBoxExtender.h"
 
 static const FName PCPluginTabName("PluginCreatorPlugin");
@@ -45,6 +43,9 @@ void FPluginCreatorModule::StartupModule()
 		FOnSpawnTab::CreateRaw(this, &FPluginCreatorModule::HandleSpawnPluginCreatorTab))
 		.SetDisplayName(LOCTEXT("PluginCreatorTabTitle", "Plugin Creator"))
 		.SetMenuType(ETabSpawnerMenuType::Hidden);
+
+	FPropertyEditorModule& PropertyModule = FModuleManager::LoadModuleChecked<FPropertyEditorModule>("PropertyEditor");
+	PropertyModule.RegisterCustomClassLayout("PluginDescriptorObject", FOnGetDetailCustomizationInstance::CreateStatic(&FPluginDescriptorDetails::MakeInstance));
 }
 
 void FPluginCreatorModule::AddMenuExtension(FMenuBuilder& Builder)
diff --git a/Engine/Plugins/Editor/PluginCreator/Source/PluginCreator/Private/UI/SLayoutCreator.h b/Engine/Plugins/Editor/PluginCreator/Source/PluginCreator/Private/UI/SLayoutCreator.h
deleted file mode 100644
diff --git a/Engine/Plugins/Editor/PluginCreator/Source/PluginCreator/Private/UI/SPluginCreatorTabContent.cpp b/Engine/Plugins/Editor/PluginCreator/Source/PluginCreator/Private/UI/SPluginCreatorTabContent.cpp
--- a/Engine/Plugins/Editor/PluginCreator/Source/PluginCreator/Private/UI/SPluginCreatorTabContent.cpp
+++ b/Engine/Plugins/Editor/PluginCreator/Source/PluginCreator/Private/UI/SPluginCreatorTabContent.cpp
@@ -11,7 +11,6 @@
 #include "SlateExtras.h"
 #include "Runtime/Slate/Private/Framework/Docking/DockingPrivate.h"
 
-#include "SLayoutCreator.h"
 #include "EditorStyle.h"
 
 #include "GameProjectUtils.h"
@@ -28,9 +27,11 @@ SPluginCreatorTabContent::SPluginCreatorTabContent()
 	const FText BlankTemplateName = LOCTEXT("BlankLabel", "Blank");
 	const FText BasicTemplateName = LOCTEXT("BasicTemplateTabLabel", "Toolbar Button");
 	const FText AdvancedTemplateName = LOCTEXT("AdvancedTemplateTabLabel", "Standalone Window");
+	const FText	BlueprintLibTemplateName = LOCTEXT("BlueprintLibTemplateLabel", "Blueprint Library");
 	const FText BlankDescription = LOCTEXT("BlankTemplateDesc", "Create a blank plugin with a minimal amount of code.\n\nChoose this if you want to set everything up from scratch or are making a non-visual plugin.\nA plugin created with this template will appear in the Editor's plugin list but will not register any buttons or menu entries.\n\nThis template will generate the following structure at:");
 	const FText BasicDescription = LOCTEXT("BasicTemplateDesc", "Create a plugin that will add a button to the toolbar in the Level Editor.\n\nStart by implementing something in the created \"OnButtonClick\" event.\n\nThis template will generate the following structure at:");
 	const FText AdvancedDescription = LOCTEXT("AdvancedTemplateDesc", "Create a plugin that will add a button to the toolbar in the Level Editor that summons an empty standalone tab window when clicked.\n\nThis template will generate the following structure at:");
+	const FText BlueprintLibDescription = LOCTEXT("BPLibTemplateDesc", "Create a plugin that will contain Blueprint Function Library. Choose this if you want to create static blueprint nodes.\n\nThis template will generate the following structure at:");
 
 	new (Templates) FPluginTemplateDescription(BlankTemplateName, BlankDescription, TEXT("BlankPluginSource"), TEXT("Blank"), false);
 	new (Templates) FPluginTemplateDescription(BasicTemplateName, BasicDescription, TEXT("BasicPluginSource"), TEXT("Basic"), true);
@@ -60,7 +61,6 @@ void SPluginCreatorTabContent::Construct(const FArguments& InArgs, TSharedPtr<SD
 
 	bUsePublicPrivateSplit = true;
 	bPluginNameIsValid = false;
-	bIsEnginePlugin = false;
 
 	OnPluginCreated = InArgs._OnPluginCreated;
 
@@ -142,25 +142,6 @@ void SPluginCreatorTabContent::Construct(const FArguments& InArgs, TSharedPtr<SD
 								]
 							]
 #endif
-
-							// Disabled because the code in PluginStyle.cpp.template doesn't handle setting the path correctly for engine versus game
-#if UE_PLUGIN_CREATOR_WIP
-							+ SVerticalBox::Slot()
-							.HAlign(HAlign_Center)
-							.VAlign(VAlign_Center)
-							.Padding(3)
-							[
-								SNew(SCheckBox)
-								.OnCheckStateChanged(this, &SPluginCreatorTabContent::OnIsEnginePluginChanged)
-								.IsChecked(ECheckBoxState::Unchecked)
-								.Content()
-								[
-									SNew(STextBlock)
-									.Text(LOCTEXT("EngineCheckboxLabel", "Should this plugin be an engine plugin?"))
-								]
-							]
-#endif
-
 							// Disable due to styling and behavior
 #if UE_PLUGIN_CREATOR_WIP
 							+ SVerticalBox::Slot()
@@ -240,19 +221,19 @@ void SPluginCreatorTabContent::Construct(const FArguments& InArgs, TSharedPtr<SD
 			]
 
 			+ SVerticalBox::Slot()
-				.AutoHeight()
-				.Padding(5)
-				.HAlign(HAlign_Right)
-				[
-					SNew(SButton)
-					.ContentPadding(5)
-					.TextStyle(FEditorStyle::Get(), "LargeText")
-					.ButtonStyle(FEditorStyle::Get(), "FlatButton.Success")
-					.IsEnabled(this, &SPluginCreatorTabContent::IsPluginNameValid)
-					.HAlign(HAlign_Center)
-					.Text(LOCTEXT("GoButtonLabel", "Create plugin"))
-					.OnClicked(this, &SPluginCreatorTabContent::OnCreatePluginClicked)
-				]
+			.AutoHeight()
+			.Padding(5)
+			.HAlign(HAlign_Right)
+			[
+				SNew(SButton)
+				.ContentPadding(5)
+				.TextStyle(FEditorStyle::Get(), "LargeText")
+				.ButtonStyle(FEditorStyle::Get(), "FlatButton.Success")
+				.IsEnabled(this, &SPluginCreatorTabContent::IsPluginNameValid)
+				.HAlign(HAlign_Center)
+				.Text(LOCTEXT("GoButtonLabel", "Create plugin"))
+				.OnClicked(this, &SPluginCreatorTabContent::OnCreatePluginClicked)
+			]
 		];
 
 		
@@ -412,12 +393,6 @@ void SPluginCreatorTabContent::OnUsePrivatePublicSplitChanged(ECheckBoxState InS
 	bUsePublicPrivateSplit = InState == ECheckBoxState::Checked;
 }
 
-
-void SPluginCreatorTabContent::OnIsEnginePluginChanged(ECheckBoxState InState)
-{
-	bIsEnginePlugin = InState == ECheckBoxState::Checked;
-}
-
 void SPluginCreatorTabContent::OnPluginNameTextChanged(const FText& InText)
 {
 	// Early exit if text is empty
@@ -483,6 +458,7 @@ void SPluginCreatorTabContent::OnPluginNameTextChanged(const FText& InText)
 	if (DescriptorObject)
 	{
 		DescriptorObject->FriendlyName = InText.ToString();
+		DescriptorObject->EditorModeSettings.Name = InText.ToString() + TEXT("EdMode");
 	}
 }
 
@@ -574,8 +550,14 @@ FReply SPluginCreatorTabContent::OnCreatePluginClicked()
 
 	FReply ReturnReply = FReply::Unhandled();
 
+	if (DescriptorObject == nullptr)
+	{
+		PopErrorNotification(LOCTEXT("FailedDescriptorObject", "Failed to create plugin. Invalid Descriptor Object."));
+		return FReply::Unhandled();
+	}
+
 	// If we're creating a game plugin, GameDir\Plugins folder may not exist so we have to create it
-	if (!bIsEnginePlugin && !IFileManager::Get().DirectoryExists(*GetPluginDestinationPath().ToString()))
+	if (!DescriptorObject->bIsEnginePlugin && !IFileManager::Get().DirectoryExists(*GetPluginDestinationPath().ToString()))
 	{
 		if (!MakeDirectory(*GetPluginDestinationPath().ToString()))
 		{
@@ -629,8 +611,18 @@ FReply SPluginCreatorTabContent::OnCreatePluginClicked()
 		bSucceeded = bSucceeded && MakeDirectory(PublicSourceFolder);
 	}
 
+	// If we're going to create editor mode, make sure that all build files will have necessary modules dependencies
+	// Only Blank plugin needs additional modules dependencies
+	// @TODO: find better way check if selected template is Blank template
+	if (DescriptorObject->bMakeEditorMode && !SelectedTemplate.bIncludeUI)
+	{
+		DescriptorObject->PrivateDependencyModuleNames.AddUnique(TEXT("InputCore"));
+		DescriptorObject->PrivateDependencyModuleNames.AddUnique(TEXT("UnrealEd"));
+		DescriptorObject->PrivateDependencyModuleNames.AddUnique(TEXT("LevelEditor"));
+	}
+
 	// Based on chosen template create build, and other source files
-	if (bSucceeded && !FPluginHelpers::CreatePluginBuildFile(PluginSourceFolder / AutoPluginName + TEXT(".Build.cs"), AutoPluginName, LocalFailReason, SelectedTemplate.OnDiskPath))
+	if (bSucceeded && !FPluginHelpers::CreatePluginBuildFile(PluginSourceFolder / AutoPluginName + TEXT(".Build.cs"), AutoPluginName, LocalFailReason, SelectedTemplate.OnDiskPath, DescriptorObject->PrivateDependencyModuleNames))
 	{
 		PopErrorNotification(FText::Format(LOCTEXT("FailedBuild", "Failed to create plugin build file. {0}"), LocalFailReason));
 		bSucceeded = false;
@@ -648,7 +640,7 @@ FReply SPluginCreatorTabContent::OnCreatePluginClicked()
 		bSucceeded = false;
 	}
 
-	if (bSucceeded && !FPluginHelpers::CreatePluginCPPFile(PrivateSourceFolder, AutoPluginName, LocalFailReason, SelectedTemplate.OnDiskPath))
+	if (bSucceeded && !FPluginHelpers::CreatePluginCPPFile(PrivateSourceFolder, AutoPluginName, LocalFailReason, SelectedTemplate.OnDiskPath, DescriptorObject->bMakeEditorMode, DescriptorObject->EditorModeSettings))
 	{
 		PopErrorNotification(FText::Format(LOCTEXT("FailedCppFile", "Failed to create plugin cpp file. {0}"), LocalFailReason));
 		bSucceeded = false;
@@ -669,6 +661,27 @@ FReply SPluginCreatorTabContent::OnCreatePluginClicked()
 		}
 	}
 
+	if (DescriptorObject)
+	{
+		if (DescriptorObject->bAddBPLibrary)
+		{
+			if (bSucceeded && !FPluginHelpers::CreateBlueprintFunctionLibraryFiles(PrivateSourceFolder, PublicSourceFolder, AutoPluginName, LocalFailReason))
+			{
+				PopErrorNotification(FText::Format(LOCTEXT("FailedBPLibraryFile", "Failed to create blueprint library files. {0}"), LocalFailReason));
+				bSucceeded = false;
+			}
+		}
+
+		if (DescriptorObject->bMakeEditorMode)
+		{
+			if (bSucceeded && !FPluginHelpers::CreatePluginEdModeFiles(PrivateSourceFolder, PublicSourceFolder, AutoPluginName, DescriptorObject->EditorModeSettings, LocalFailReason))
+			{
+				PopErrorNotification(FText::Format(LOCTEXT("FailedBPLibraryFile", "Failed to create EdMode files. {0}"), LocalFailReason));
+				bSucceeded = false;
+			}
+		}
+	}
+
 	if (bSucceeded)
 	{
 		const bool bShowFolder = (DescriptorObject != nullptr) && (DescriptorObject->bShowPluginFolder);
@@ -708,7 +721,7 @@ FText SPluginCreatorTabContent::GetPluginDestinationPath() const
 	const FString EnginePath = FPaths::EnginePluginsDir();
 	const FString GamePath = FPaths::GamePluginsDir();
 
-	return bIsEnginePlugin ? FText::FromString(IFileManager::Get().ConvertToAbsolutePathForExternalAppForWrite(*EnginePath)) 
+	return (DescriptorObject != nullptr && DescriptorObject->bIsEnginePlugin) ? FText::FromString(IFileManager::Get().ConvertToAbsolutePathForExternalAppForWrite(*EnginePath)) 
 		: FText::FromString(IFileManager::Get().ConvertToAbsolutePathForExternalAppForWrite(*GamePath));
 }
 
diff --git a/Engine/Plugins/Editor/PluginCreator/Source/PluginCreator/Private/UI/SPluginCreatorTabContent.h b/Engine/Plugins/Editor/PluginCreator/Source/PluginCreator/Private/UI/SPluginCreatorTabContent.h
--- a/Engine/Plugins/Editor/PluginCreator/Source/PluginCreator/Private/UI/SPluginCreatorTabContent.h
+++ b/Engine/Plugins/Editor/PluginCreator/Source/PluginCreator/Private/UI/SPluginCreatorTabContent.h
@@ -92,8 +92,6 @@ public:
 
 	void OnUsePrivatePublicSplitChanged(ECheckBoxState InState);
 
-	void OnIsEnginePluginChanged(ECheckBoxState InState);
-
 	TSharedRef<SWidget> GetPathTextWidget();
 
 	FText GetPluginDestinationPath() const;
@@ -135,12 +133,6 @@ private:
 	 */
 	bool bUsePublicPrivateSplit;
 
-	/** Should new plugin be considered as 'Engine' plugin?
-	 *	If so, it will be put under Engine/Plugins directory.
-	 *	False by default
-	 */
-	bool bIsEnginePlugin;
-
 	/** Textbox widget that user will put plugin name in*/
 	TSharedPtr<SEditableTextBox> PluginNameTextBox;
 
diff --git a/Engine/Plugins/Editor/PluginCreator/Templates/Advanced/PluginModule.cpp.template b/Engine/Plugins/Editor/PluginCreator/Templates/Advanced/PluginModule.cpp.template
--- a/Engine/Plugins/Editor/PluginCreator/Templates/Advanced/PluginModule.cpp.template
+++ b/Engine/Plugins/Editor/PluginCreator/Templates/Advanced/PluginModule.cpp.template
@@ -51,6 +51,8 @@ void F%PLUGIN_NAME%Module::StartupModule()
 	FGlobalTabmanager::Get()->RegisterNomadTabSpawner(%PLUGIN_NAME%TabName, FOnSpawnTab::CreateRaw(this, &F%PLUGIN_NAME%Module::OnSpawnPluginTab))
 		.SetDisplayName(LOCTEXT("F%PLUGIN_NAME%TabTitle", "%PLUGIN_NAME%"))
 		.SetMenuType(ETabSpawnerMenuType::Hidden);
+		
+	%REGISTER_EDMODE%
 }
 
 void F%PLUGIN_NAME%Module::ShutdownModule()
@@ -58,6 +60,8 @@ void F%PLUGIN_NAME%Module::ShutdownModule()
 	// This function may be called during shutdown to clean up your module.  For modules that support dynamic reloading,
 	// we call this function before unloading the module.
 	
+	%UNREGISTER_EDMODE%
+	
 	F%PLUGIN_NAME%Style::Shutdown();
 
 	F%PLUGIN_NAME%Commands::Unregister();
diff --git a/Engine/Plugins/Editor/PluginCreator/Templates/Basic/PluginModule.cpp.template b/Engine/Plugins/Editor/PluginCreator/Templates/Basic/PluginModule.cpp.template
--- a/Engine/Plugins/Editor/PluginCreator/Templates/Basic/PluginModule.cpp.template
+++ b/Engine/Plugins/Editor/PluginCreator/Templates/Basic/PluginModule.cpp.template
@@ -47,6 +47,8 @@ void F%PLUGIN_NAME%Module::StartupModule()
 		
 		LevelEditorModule.GetToolBarExtensibilityManager()->AddExtender(ToolbarExtender);
 	}
+	
+	%REGISTER_EDMODE%
 }
 
 void F%PLUGIN_NAME%Module::ShutdownModule()
@@ -54,6 +56,8 @@ void F%PLUGIN_NAME%Module::ShutdownModule()
 	// This function may be called during shutdown to clean up your module.  For modules that support dynamic reloading,
 	// we call this function before unloading the module.
 	
+	%UNREGISTER_EDMODE%
+	
 	F%PLUGIN_NAME%Style::Shutdown();
 
 	F%PLUGIN_NAME%Commands::Unregister();
diff --git a/Engine/Plugins/Editor/PluginCreator/Templates/Blank/PluginModule.cpp.template b/Engine/Plugins/Editor/PluginCreator/Templates/Blank/PluginModule.cpp.template
--- a/Engine/Plugins/Editor/PluginCreator/Templates/Blank/PluginModule.cpp.template
+++ b/Engine/Plugins/Editor/PluginCreator/Templates/Blank/PluginModule.cpp.template
@@ -9,12 +9,16 @@
 void F%PLUGIN_NAME%Module::StartupModule()
 {
 	// This code will execute after your module is loaded into memory; the exact timing is specified in the .uplugin file per-module
+	
+	%REGISTER_EDMODE%
 }
 
 void F%PLUGIN_NAME%Module::ShutdownModule()
 {
 	// This function may be called during shutdown to clean up your module.  For modules that support dynamic reloading,
 	// we call this function before unloading the module.
+	
+	%UNREGISTER_EDMODE%
 }
 
 #undef LOCTEXT_NAMESPACE
diff --git a/Engine/Plugins/Editor/PluginCreator/Templates/BlueprintFunctionLibrary.cpp.template b/Engine/Plugins/Editor/PluginCreator/Templates/BlueprintFunctionLibrary.cpp.template
new file mode 100644
--- /dev/null
+++ b/Engine/Plugins/Editor/PluginCreator/Templates/BlueprintFunctionLibrary.cpp.template
@@ -0,0 +1,14 @@
+#include "%PLUGIN_NAME%PrivatePCH.h"
+#include "%PLUGIN_NAME%BPLibrary.h"
+
+U%PLUGIN_NAME%BPLibrary::U%PLUGIN_NAME%BPLibrary(const FObjectInitializer& ObjectInitializer) 
+: Super(ObjectInitializer)
+{
+
+}
+
+float U%PLUGIN_NAME%BPLibrary::%PLUGIN_NAME%SampleFunction(float Param)
+{
+	return -1;
+}
+
diff --git a/Engine/Plugins/Editor/PluginCreator/Templates/BlueprintFunctionLibrary.h.template b/Engine/Plugins/Editor/PluginCreator/Templates/BlueprintFunctionLibrary.h.template
new file mode 100644
--- /dev/null
+++ b/Engine/Plugins/Editor/PluginCreator/Templates/BlueprintFunctionLibrary.h.template
@@ -0,0 +1,29 @@
+#pragma once
+#include "Engine.h"
+#include "%PLUGIN_NAME%BPLibrary.generated.h"
+
+/* 
+*	Function library class.
+*	Each function in it is expected to be static and represents blueprint node that can be called in any blueprint.
+*
+*	When declaring function you can define metadata for the node. Key function specifiers will be BlueprintPure and BlueprintCallable.
+*	BlueprintPure - means the function does not affect the owning object in any way and thus creates a node without Exec pins.
+*	BlueprintCallable - makes a function which can be executed in Blueprints - Thus it has Exec pins.
+*	DisplayName - full name of the node, shown when you mouse over the node and in the blueprint drop down menu.
+*				Its lets you name the node using characters not allowed in C++ function names.
+*	CompactNodeTitle - the word(s) that appear on the node.
+*	Keywords -	the list of keywords that helps you to find node when you search for it using Blueprint drop-down menu. 
+*				Good example is "Print String" node which you can find also by using keyword "log".
+*	Category -	the category your node will be under in the Blueprint drop-down menu.
+*
+*	For more info on custom blueprint nodes visit documentation:
+*	https://wiki.unrealengine.com/Custom_Blueprint_Node_Creation
+*/
+UCLASS()
+class U%PLUGIN_NAME%BPLibrary : public UBlueprintFunctionLibrary
+{
+	GENERATED_UCLASS_BODY()
+
+	UFUNCTION(BlueprintCallable, meta = (DisplayName = "Execute Sample function", Keywords = "%PLUGIN_NAME% sample test testing"), Category = "%PLUGIN_NAME%Testing")
+	static float %PLUGIN_NAME%SampleFunction(float Param);
+};
diff --git a/Engine/Plugins/Editor/PluginCreator/Templates/EdMode.cpp.template b/Engine/Plugins/Editor/PluginCreator/Templates/EdMode.cpp.template
new file mode 100644
--- /dev/null
+++ b/Engine/Plugins/Editor/PluginCreator/Templates/EdMode.cpp.template
@@ -0,0 +1,56 @@
+#include "%PLUGIN_NAME%PrivatePCH.h"
+#include "%MODE_NAME%.h"
+
+%BEGIN_BLOCK_COMMENT%
+#include "%MODE_NAME%Toolkit.h"
+#include "Toolkits/ToolkitManager.h"
+%END_BLOCK_COMMENT%
+
+const FEditorModeID F%MODE_NAME%::EM_%MODE_NAME%Id = TEXT("EM_%MODE_NAME%");
+
+
+F%MODE_NAME%::F%MODE_NAME%()
+{
+
+}
+
+F%MODE_NAME%::~F%MODE_NAME%()
+{
+
+}
+
+void F%MODE_NAME%::Enter()
+{
+	FEdMode::Enter();
+
+	%BEGIN_BLOCK_COMMENT%
+	if (!Toolkit.IsValid() && UsesToolkits())
+	{
+		Toolkit = MakeShareable(new F%MODE_NAME%Toolkit);
+		Toolkit->Init(Owner->GetToolkitHost());
+	}
+	%END_BLOCK_COMMENT%
+}
+
+void F%MODE_NAME%::Exit()
+{
+	%BEGIN_BLOCK_COMMENT%
+	if (Toolkit.IsValid())
+	{
+		FToolkitManager::Get().CloseToolkit(Toolkit.ToSharedRef());
+		Toolkit.Reset();
+	}
+	%END_BLOCK_COMMENT%
+
+	// Call base Exit method to ensure proper cleanup
+	FEdMode::Exit();
+}
+
+bool F%MODE_NAME%::UsesToolkits() const
+{
+	return %USES_TOOLKIT%;
+}
+
+
+
+
diff --git a/Engine/Plugins/Editor/PluginCreator/Templates/EdMode.h.template b/Engine/Plugins/Editor/PluginCreator/Templates/EdMode.h.template
new file mode 100644
--- /dev/null
+++ b/Engine/Plugins/Editor/PluginCreator/Templates/EdMode.h.template
@@ -0,0 +1,22 @@
+#pragma once
+
+#include "UnrealEd.h" 
+#include "Editor.h"
+
+class F%MODE_NAME% : public FEdMode
+{
+public:
+	const static FEditorModeID EM_%MODE_NAME%Id;
+public:
+	F%MODE_NAME%();
+	virtual ~F%MODE_NAME%();
+
+	// FEdMode interface
+	virtual void Enter() override;
+	virtual void Exit() override;
+	//virtual void Tick(FEditorViewportClient* ViewportClient, float DeltaTime) override;
+	//virtual void Render(const FSceneView* View, FViewport* Viewport, FPrimitiveDrawInterface* PDI) override;
+	//virtual void ActorSelectionChangeNotify() override;
+	bool UsesToolkits() const override;
+	// End of FEdMode interface
+};
diff --git a/Engine/Plugins/Editor/PluginCreator/Templates/Toolkit.cpp.template b/Engine/Plugins/Editor/PluginCreator/Templates/Toolkit.cpp.template
new file mode 100644
--- /dev/null
+++ b/Engine/Plugins/Editor/PluginCreator/Templates/Toolkit.cpp.template
@@ -0,0 +1,27 @@
+#include "%PLUGIN_NAME%PrivatePCH.h"
+#include "%MODE_NAME%.h"
+#include "%MODE_NAME%Toolkit.h"
+
+#define LOCTEXT_NAMESPACE "F%MODE_NAME%Toolkit"
+
+F%MODE_NAME%Toolkit::F%MODE_NAME%Toolkit()
+{
+	%TOOLKIT_SAMPLE_UI%
+}
+
+FName F%MODE_NAME%Toolkit::GetToolkitFName() const
+{
+	return FName("%MODE_NAME%");
+}
+
+FText F%MODE_NAME%Toolkit::GetBaseToolkitName() const
+{
+	return NSLOCTEXT("%MODE_NAME%Toolkit", "DisplayName", "%MODE_NAME% Tool");
+}
+
+class FEdMode* F%MODE_NAME%Toolkit::GetEditorMode() const
+{
+	return GLevelEditorModeTools().GetActiveMode(F%MODE_NAME%::EM_%MODE_NAME%Id);
+}
+
+#undef LOCTEXT_NAMESPACE
diff --git a/Engine/Plugins/Editor/PluginCreator/Templates/Toolkit.h.template b/Engine/Plugins/Editor/PluginCreator/Templates/Toolkit.h.template
new file mode 100644
--- /dev/null
+++ b/Engine/Plugins/Editor/PluginCreator/Templates/Toolkit.h.template
@@ -0,0 +1,21 @@
+#pragma once
+
+#include "Editor/UnrealEd/Public/Toolkits/BaseToolkit.h"
+
+
+class F%MODE_NAME%Toolkit : public FModeToolkit
+{
+public:
+
+	F%MODE_NAME%Toolkit();
+
+	/** IToolkit interface */
+	virtual FName GetToolkitFName() const override;
+	virtual FText GetBaseToolkitName() const override;
+	virtual class FEdMode* GetEditorMode() const override;
+	virtual TSharedPtr<class SWidget> GetInlineContent() const override { return ToolkitWidget; }
+
+private:
+
+	TSharedPtr<SWidget> ToolkitWidget;
+};
diff --git a/Engine/Plugins/Editor/PluginCreator/Templates/ToolkitUI.template b/Engine/Plugins/Editor/PluginCreator/Templates/ToolkitUI.template
new file mode 100644
--- /dev/null
+++ b/Engine/Plugins/Editor/PluginCreator/Templates/ToolkitUI.template
@@ -0,0 +1,87 @@
+struct Locals
+	{
+		static bool IsWidgetEnabled()
+		{
+			return GEditor->GetSelectedActors()->Num() != 0;
+		}
+
+		static FReply OnButtonClick(FVector InOffset)
+		{
+			USelection* SelectedActors = GEditor->GetSelectedActors();
+
+			// Let editor know that we're about to do something that we want to undo/redo
+			GEditor->BeginTransaction(LOCTEXT("MoveActorsTransactionName", "MoveActors"));
+
+			// For each selected actor
+			for (FSelectionIterator Iter(*SelectedActors); Iter; ++Iter)
+			{
+				if (AActor* LevelActor = Cast<AActor>(*Iter))
+				{
+					// Register actor in opened transaction (undo/redo)
+					LevelActor->Modify();
+					// Move actor to given location
+					LevelActor->TeleportTo(LevelActor->GetActorLocation() + InOffset, FRotator(0, 0, 0));
+				}
+			}
+
+			// We're done moving actors so close transaction
+			GEditor->EndTransaction();
+
+			return FReply::Handled();
+		}
+
+		static TSharedRef<SWidget> MakeButton(FText InLabel, const FVector InOffset)
+		{
+			return SNew(SButton)
+				.Text(InLabel)
+				.OnClicked_Static(&Locals::OnButtonClick, InOffset);
+		}
+	};
+
+	const float Factor = 256.0f;
+
+	SAssignNew(ToolkitWidget, SBorder)
+	.HAlign(HAlign_Center)
+	.Padding(25)
+	.IsEnabled_Static(&Locals::IsWidgetEnabled)
+	[
+		SNew(SVerticalBox)
+		+SVerticalBox::Slot()
+		.AutoHeight()
+		.HAlign(HAlign_Center)
+		.Padding(50)
+		[
+			SNew(STextBlock)
+			.AutoWrapText(true)
+			.Text(LOCTEXT("HelperLabel", "Select some actors and move them around using buttons below"))
+		]
+		+SVerticalBox::Slot()
+		.HAlign(HAlign_Center)
+		.AutoHeight()
+		[
+			Locals::MakeButton(LOCTEXT("UpButtonLabel", "Up"), FVector(0, 0, Factor))
+		]
+		+SVerticalBox::Slot()
+		.HAlign(HAlign_Center)
+		.AutoHeight()
+		[
+			SNew(SHorizontalBox)
+			+SHorizontalBox::Slot()
+			.AutoWidth()
+			[
+				Locals::MakeButton(LOCTEXT("LeftButtonLabel", "Left"), FVector(0, -Factor, 0))
+			]
+			+SHorizontalBox::Slot()
+			.AutoWidth()
+			[
+				Locals::MakeButton(LOCTEXT("RightButtonLabel", "Right"), FVector(0, Factor, 0))
+			]
+		]
+		+SVerticalBox::Slot()
+		.HAlign(HAlign_Center)
+		.AutoHeight()
+		[
+			Locals::MakeButton(LOCTEXT("DownButtonLabel", "Down"), FVector(0, 0, -Factor))
+		]
+		
+	];
\ No newline at end of file
