Fixing non-unity build error

diff --git a/Engine/Plugins/Editor/PluginCreator/Source/PluginCreator/PluginCreator.Build.cs b/Engine/Plugins/Editor/PluginCreator/Source/PluginCreator/PluginCreator.Build.cs
--- a/Engine/Plugins/Editor/PluginCreator/Source/PluginCreator/PluginCreator.Build.cs
+++ b/Engine/Plugins/Editor/PluginCreator/Source/PluginCreator/PluginCreator.Build.cs
@@ -13,6 +13,7 @@ public class PluginCreator : ModuleRules
 			{ 
 				"DesktopPlatform", 
  				"HardwareTargeting",
+				"PropertyEditor",
 			});
 
 		PrivateDependencyModuleNames.AddRange(new string[]
diff --git a/Engine/Plugins/Editor/PluginCreator/Source/PluginCreator/Private/Helpers/PluginDescriptorDetails.cpp b/Engine/Plugins/Editor/PluginCreator/Source/PluginCreator/Private/Helpers/PluginDescriptorDetails.cpp
--- a/Engine/Plugins/Editor/PluginCreator/Source/PluginCreator/Private/Helpers/PluginDescriptorDetails.cpp
+++ b/Engine/Plugins/Editor/PluginCreator/Source/PluginCreator/Private/Helpers/PluginDescriptorDetails.cpp
@@ -1,6 +1,8 @@
 #include "PluginCreatorPrivatePCH.h"
 #include "PluginDescriptorDetails.h"
 #include "PluginDescriptorObject.h"
+#include "DetailLayoutBuilder.h"
+#include "PropertyHandle.h"
 
 #define LOCTEXT_NAMESPACE "PluginDescriptorDetails"
 
diff --git a/Engine/Plugins/Editor/PluginCreator/Source/PluginCreator/Private/Helpers/PluginDescriptorDetails.h b/Engine/Plugins/Editor/PluginCreator/Source/PluginCreator/Private/Helpers/PluginDescriptorDetails.h
--- a/Engine/Plugins/Editor/PluginCreator/Source/PluginCreator/Private/Helpers/PluginDescriptorDetails.h
+++ b/Engine/Plugins/Editor/PluginCreator/Source/PluginCreator/Private/Helpers/PluginDescriptorDetails.h
@@ -1,5 +1,7 @@
 #pragma once
 
+#include "IDetailCustomization.h"
+
 class FPluginDescriptorDetails : public IDetailCustomization
 {
 public:
