Disable Plugin Creator in content only projects as it currently wont generate projects/compile any plugins that are created.

UE-15150

diff --git a/Engine/Plugins/Editor/PluginCreator/Source/PluginCreator/Private/PluginCreatorPlugin.cpp b/Engine/Plugins/Editor/PluginCreator/Source/PluginCreator/Private/PluginCreatorPlugin.cpp
--- a/Engine/Plugins/Editor/PluginCreator/Source/PluginCreator/Private/PluginCreatorPlugin.cpp
+++ b/Engine/Plugins/Editor/PluginCreator/Source/PluginCreator/Private/PluginCreatorPlugin.cpp
@@ -8,6 +8,8 @@
 #include "PluginCreatorStyle.h"
 #include "PluginDescriptorDetails.h"
 #include "MultiBoxExtender.h"
+#include "IProjectManager.h"
+#include "GameProjectUtils.h"
 
 static const FName PCPluginTabName("PluginCreatorPlugin");
 
@@ -19,6 +21,14 @@ DEFINE_LOG_CATEGORY(PluginCreatorPluginLog);
 
 void FPluginCreatorModule::StartupModule()
 {
+	// Disallow this plugin in content only projects as they won't compile
+	const FProjectDescriptor* CurrentProject = IProjectManager::Get().GetCurrentProject();
+	if (CurrentProject == nullptr || CurrentProject->Modules.Num() == 0 || !GameProjectUtils::ProjectHasCodeFiles())
+	{
+		UE_LOG(PluginCreatorPluginLog, Warning, TEXT("Plugin Creator module disabled for content only projects"));
+		return;
+	}
+
 	FPluginCreatorStyle::Initialize();
 	FPluginCreatorStyle::ReloadTextures();
 
diff --git a/Engine/Plugins/Editor/PluginCreator/Source/PluginCreator/Private/UI/PluginCreatorStyle.cpp b/Engine/Plugins/Editor/PluginCreator/Source/PluginCreator/Private/UI/PluginCreatorStyle.cpp
--- a/Engine/Plugins/Editor/PluginCreator/Source/PluginCreator/Private/UI/PluginCreatorStyle.cpp
+++ b/Engine/Plugins/Editor/PluginCreator/Source/PluginCreator/Private/UI/PluginCreatorStyle.cpp
@@ -18,9 +18,12 @@ void FPluginCreatorStyle::Initialize()
 
 void FPluginCreatorStyle::Shutdown()
 {
-	FSlateStyleRegistry::UnRegisterSlateStyle(*PCStyleInstance);
-	ensure(PCStyleInstance.IsUnique());
-	PCStyleInstance.Reset();
+	if (PCStyleInstance.IsValid())
+	{
+		FSlateStyleRegistry::UnRegisterSlateStyle(*PCStyleInstance);
+		ensure(PCStyleInstance.IsUnique());
+		PCStyleInstance.Reset();
+	}
 }
 
 FName FPluginCreatorStyle::GetStyleSetName()
